Porting AnimNode_OffsetRootBone to UAF OffsetRootBoneTrait

#virtualized

diff --git a/Engine/Plugins/Experimental/UAF/UAFWarping/Source/UAFWarping/Private/OffsetRootBoneTrait.cpp b/Engine/Plugins/Experimental/UAF/UAFWarping/Source/UAFWarping/Private/OffsetRootBoneTrait.cpp
new file mode 100644
--- /dev/null
+++ b/Engine/Plugins/Experimental/UAF/UAFWarping/Source/UAFWarping/Private/OffsetRootBoneTrait.cpp
@@ -0,0 +1,367 @@
+// Copyright Epic Games, Inc. All Rights Reserved.
+
+#include "OffsetRootBoneTrait.h"
+
+#include "AnimNextWarpingLog.h"
+#include "Animation/AnimRootMotionProvider.h"
+#include "Animation/SpringMath.h"
+#include "EvaluationVM/EvaluationVM.h"
+#include "VisualLogger/VisualLogger.h"
+
+#if ENABLE_ANIM_DEBUG
+TAutoConsoleVariable<int32> CVarAnimNodeOffsetRootBoneDebug(TEXT("a.UAF.OffsetRootBone.Debug"), 0, TEXT("Turn on visualization debugging for Offset Root Bone"));
+TAutoConsoleVariable<int32> CVarAnimNodeOffsetRootBoneEnable(TEXT("a.UAF.OffsetRootBone.Enable"), 1, TEXT("Toggle Offset Root Bone"));
+#endif
+
+namespace UE::UAF::OffsetRootBone
+{
+
+static bool ShouldExtractRootMotion(const EUAFOffsetRootBoneMode OffsetMode)
+{
+	switch (OffsetMode)
+	{
+	case EUAFOffsetRootBoneMode::Accumulate:
+	case EUAFOffsetRootBoneMode::Interpolate:
+		return true;
+	case EUAFOffsetRootBoneMode::Release:
+	default:
+		return false;
+	}
+}
+
+static bool ShouldCounterComponentDelta(const EUAFOffsetRootBoneMode OffsetMode)
+{
+	switch (OffsetMode)
+	{
+	case EUAFOffsetRootBoneMode::Accumulate:
+	case EUAFOffsetRootBoneMode::Interpolate:
+		return false;
+	case EUAFOffsetRootBoneMode::Release:
+	default:
+		return true;
+	}
+}
+}
+
+////////////////////////////////////////////////////////////////////////////////////////////////////////////////////////
+// FOffsetRootBoneTrait
+
+namespace UE::UAF
+{
+	AUTO_REGISTER_ANIM_TRAIT(FOffsetRootBoneTrait)
+
+	#define TRAIT_INTERFACE_ENUMERATOR(GeneratorMacro) \
+		GeneratorMacro(IUpdate) \
+		GeneratorMacro(IEvaluate) \
+		GeneratorMacro(IHierarchy) \
+
+	// Trait implementation boilerplate
+	GENERATE_ANIM_TRAIT_IMPLEMENTATION(FOffsetRootBoneTrait, TRAIT_INTERFACE_ENUMERATOR, NULL_ANIM_TRAIT_INTERFACE_ENUMERATOR, NULL_ANIM_TRAIT_EVENT_ENUMERATOR)
+	#undef TRAIT_INTERFACE_ENUMERATOR
+
+	void FOffsetRootBoneTrait::PreUpdate(FUpdateTraversalContext& Context, const TTraitBinding<IUpdate>& Binding, const FTraitUpdateState& TraitState) const
+	{
+		FInstanceData* InstanceData = Binding.GetInstanceData<FInstanceData>();
+		check(InstanceData);
+
+		const FSharedData* SharedData = Binding.GetSharedData<FSharedData>();
+		check(SharedData);
+		
+		InstanceData->DeltaTime = TraitState.GetDeltaTime();
+
+		if (!InstanceData->Input.IsValid())
+		{
+			InstanceData->Input = Context.AllocateNodeInstance(Binding, SharedData->Input);
+		}
+
+		IUpdate::PreUpdate(Context, Binding, TraitState);
+	}
+
+	void FOffsetRootBoneTrait::PostEvaluate(FEvaluateTraversalContext& Context, const TTraitBinding<IEvaluate>& Binding) const
+	{
+		IEvaluate::PostEvaluate(Context, Binding);
+
+		const FSharedData* SharedData = Binding.GetSharedData<FSharedData>();
+		check(SharedData);
+
+		FInstanceData* InstanceData = Binding.GetInstanceData<FInstanceData>();
+		check(InstanceData);
+
+#if ENABLE_ANIM_DEBUG
+		if(CVarAnimNodeOffsetRootBoneEnable.GetValueOnAnyThread() == 0)
+		{
+			InstanceData->bIsFirstUpdate = true;
+			return;
+		}
+#endif
+
+		// Need to copy over latent values from shared data to instance data in order to sample all the pin values (via the getters)
+		// Todo: wish we didn't have to do this
+		InstanceData->Alpha = SharedData->GetAlpha(Binding);
+		InstanceData->HostTransformWorld = SharedData->GetHostTransformWorld(Binding);
+		InstanceData->TranslationMode = SharedData->GetTranslationMode(Binding);
+		InstanceData->RotationMode = SharedData->GetRotationMode(Binding);
+		InstanceData->MaxTranslationError = SharedData->GetMaxTranslationError(Binding);
+		InstanceData->MaxRotationErrorDegrees = SharedData->GetMaxRotationErrorDegrees(Binding);
+		InstanceData->TranslationSmoothingTime = SharedData->GetTranslationSmoothingTime(Binding);
+		InstanceData->RotationSmoothingTime = SharedData->GetRotationSmoothingTime(Binding);
+
+#if ENABLE_ANIM_DEBUG 
+		InstanceData->HostObject = Context.GetHostObject();
+#endif // ENABLE_ANIM_DEBUG 
+
+		Context.AppendTask(FAnimNextOffsetRootBoneTask::Make(InstanceData, SharedData));
+	}
+
+// --- IHierarchy impl --- 
+uint32 FOffsetRootBoneTrait::GetNumChildren(const FExecutionContext& Context, const TTraitBinding<IHierarchy>& Binding) const
+	{
+		return 1;
+	}
+
+void FOffsetRootBoneTrait::GetChildren(const FExecutionContext& Context, const TTraitBinding<IHierarchy>& Binding, FChildrenArray& Children) const
+	{
+		const FInstanceData* InstanceData = Binding.GetInstanceData<FInstanceData>();
+		Children.Add(InstanceData->Input);
+	}
+} // namespace UE::UAF
+
+
+////////////////////////////////////////////////////////////////////////////////////////////////////////////////////////
+// FAnimNextOffsetRootBoneTask
+
+FAnimNextOffsetRootBoneTask FAnimNextOffsetRootBoneTask::Make(UE::UAF::FOffsetRootBoneTrait::FInstanceData* InstanceData, const UE::UAF::FOffsetRootBoneTrait::FSharedData* SharedData)
+{
+	FAnimNextOffsetRootBoneTask Task;
+	Task.InstanceData = InstanceData;
+	Task.SharedData = SharedData;
+	return Task;
+}
+
+void FAnimNextOffsetRootBoneTask::Execute(UE::UAF::FEvaluationVM& VM) const
+{
+	using namespace UE::UAF;
+	QUICK_SCOPE_CYCLE_COUNTER(FAnimNextOffsetRootBoneTask_Execute);
+	
+	if (InstanceData->DeltaTime == 0.f)
+	{
+		return;
+	}
+
+	if (InstanceData->Alpha == 0.f)
+	{
+		return;
+	}
+
+	// Extract the root motion for this frame
+	TUniquePtr<FKeyframeState>* Keyframe = VM.PeekValueMutable<TUniquePtr<FKeyframeState>>(KEYFRAME_STACK_NAME, 0);
+	if (Keyframe == nullptr)
+	{
+		return;
+	}
+	
+	FTransform RootMotionTransformDelta = FTransform::Identity;
+	const UE::Anim::IAnimRootMotionProvider* RootMotionProvider = UE::Anim::IAnimRootMotionProvider::Get();
+	if (!RootMotionProvider)
+	{
+		UE_LOG(LogAnimNextWarping, Error, TEXT("FAnimNextOffsetRootBoneTask::Execute, missing RootMotionProvider"));
+		return;
+	}
+
+	const FTransform InputBoneTransform = Keyframe->Get()->Pose.LocalTransforms[UE::UAF::FReferencePose::RootBoneIndex];
+
+	const FTransform ComponentTransform = InstanceData->HostTransformWorld;
+
+	if (InstanceData->bIsFirstUpdate)
+	{
+		InstanceData->SimulatedTranslation = ComponentTransform.GetTranslation();
+		InstanceData->SimulatedRotation = ComponentTransform.GetRotation();
+	}
+
+	// Note we set last to current on first update
+	const FTransform LastComponentTransform = InstanceData->bIsFirstUpdate ? ComponentTransform : InstanceData->LastMeshComponentTransformWorld;
+	InstanceData->bIsFirstUpdate = false;
+	
+	// Cycle last transform
+	InstanceData->LastMeshComponentTransformWorld = InstanceData->HostTransformWorld;
+
+	const EUAFOffsetRootBoneMode CurrentTranslationMode = InstanceData->TranslationMode;
+	const EUAFOffsetRootBoneMode CurrentRotationMode = InstanceData->RotationMode;;
+
+	bool bShouldConsumeTranslationOffset = UE::UAF::OffsetRootBone::ShouldExtractRootMotion(CurrentTranslationMode);
+	bool bShouldConsumeRotationOffset = UE::UAF::OffsetRootBone::ShouldExtractRootMotion(CurrentRotationMode);
+
+	if (RootMotionProvider->ExtractRootMotion(Keyframe->Get()->Attributes, RootMotionTransformDelta) == false)
+	{
+		return;
+	}
+
+	RootMotionTransformDelta.NormalizeRotation();
+	
+	FTransform ConsumedRootMotionDelta = FTransform::Identity;
+
+	if (bShouldConsumeTranslationOffset)
+	{
+		// Grab root motion translation from the root motion attribute
+		ConsumedRootMotionDelta.SetTranslation(RootMotionTransformDelta.GetTranslation());
+	}
+	if (bShouldConsumeRotationOffset)
+	{
+		// Grab root motion rotation from the root motion attribute
+		ConsumedRootMotionDelta.SetRotation(RootMotionTransformDelta.GetRotation());
+	}
+
+	if (UE::UAF::OffsetRootBone::ShouldCounterComponentDelta(CurrentRotationMode))
+	{
+		// Accumulate the rotation component delta into the simulated rotation, to keep component and offset in sync.
+		const FQuat ComponentRotationDelta = LastComponentTransform.GetRotation().Inverse() * ComponentTransform.GetRotation();
+		InstanceData->SimulatedRotation = ComponentRotationDelta * InstanceData->SimulatedRotation;
+	}
+	if (UE::UAF::OffsetRootBone::ShouldCounterComponentDelta(CurrentTranslationMode))
+	{
+		// Accumulate the translation component delta into the simulated translation, to keep component and offset in sync.
+		const FVector ComponentTranslationDelta = ComponentTransform.GetLocation() - LastComponentTransform.GetLocation();
+		InstanceData->SimulatedTranslation += ComponentTranslationDelta;
+	}
+
+	FTransform SimulatedTransform(InstanceData->SimulatedRotation, InstanceData->SimulatedTranslation);
+	// Apply the root motion delta
+	SimulatedTransform = ConsumedRootMotionDelta * SimulatedTransform;
+
+	InstanceData->SimulatedTranslation = SimulatedTransform.GetLocation();
+	InstanceData->SimulatedRotation = SimulatedTransform.GetRotation();
+
+	if (InstanceData->TranslationMode == EUAFOffsetRootBoneMode::Release ||
+		InstanceData->TranslationMode == EUAFOffsetRootBoneMode::Interpolate)
+	{
+		FVector TranslationOffset = ComponentTransform.GetLocation() - InstanceData->SimulatedTranslation;
+
+		// Blend out translation offset
+		FVector TranslationOffsetDelta = FVector::ZeroVector;
+		FMath::ExponentialSmoothingApprox(TranslationOffsetDelta,TranslationOffset, InstanceData->DeltaTime, SharedData->TranslationSmoothingTime);
+
+		if (SharedData->bClampToTranslationVelocity)
+		{
+			const float RootMotionDelta = RootMotionTransformDelta.GetLocation().Size();
+			const float MaxDelta = 	SharedData->TranslationSpeedRatio * RootMotionDelta;
+
+			const float AdjustmentDelta = TranslationOffsetDelta.Size();
+			if (AdjustmentDelta > MaxDelta)
+			{
+				TranslationOffsetDelta = MaxDelta * TranslationOffsetDelta.GetSafeNormal2D();
+			}
+		}
+		
+		InstanceData->SimulatedTranslation = InstanceData->SimulatedTranslation + TranslationOffsetDelta;
+	}
+
+	if (InstanceData->RotationMode == EUAFOffsetRootBoneMode::Release ||
+		InstanceData->RotationMode == EUAFOffsetRootBoneMode::Interpolate)
+	{
+		FQuat RotationOffset = ComponentTransform.GetRotation() * InstanceData->SimulatedRotation.Inverse();
+		RotationOffset.EnforceShortestArcWith(FQuat::Identity);
+		FQuat RotationOffsetDelta = FQuat::Identity;
+		SpringMath::ExponentialSmoothingApproxQuat( RotationOffsetDelta, RotationOffset, InstanceData->DeltaTime, SharedData->RotationSmoothingTime);
+
+		if (SharedData->bClampToRotationVelocity)
+		{
+			float RotationMotionAngleDelta;
+			FVector RootMotionRotationAxis;
+			RootMotionTransformDelta.GetRotation().ToAxisAndAngle(RootMotionRotationAxis, RotationMotionAngleDelta);
+
+			float MaxRotationAngle = SharedData->RotationSpeedRatio * RotationMotionAngleDelta;
+
+			FVector DeltaAxis;
+			float DeltaAngle;
+			RotationOffsetDelta.ToAxisAndAngle(DeltaAxis, DeltaAngle);
+
+			if (DeltaAngle > MaxRotationAngle)
+			{
+				RotationOffsetDelta = FQuat(DeltaAxis, MaxRotationAngle);
+			}
+		}
+
+		InstanceData->SimulatedRotation = RotationOffsetDelta * InstanceData->SimulatedRotation;
+	}
+
+	if (SharedData->MaxTranslationError >= 0.0f)
+	{
+		FVector TranslationOffset = ComponentTransform.GetLocation() - InstanceData->SimulatedTranslation;
+		const float TranslationOffsetSizeSquared = TranslationOffset.SizeSquared();
+		if (TranslationOffsetSizeSquared > (SharedData->MaxTranslationError * SharedData->MaxTranslationError))
+		{
+			TranslationOffset = TranslationOffset.GetClampedToMaxSize(SharedData->MaxTranslationError);
+			InstanceData->SimulatedTranslation = ComponentTransform.GetLocation() - TranslationOffset;
+		}
+	}
+
+	const float MaxAngleRadians = FMath::DegreesToRadians(SharedData->MaxRotationErrorDegrees);
+	if (MaxAngleRadians >= 0.0f)
+	{
+		FQuat RotationOffset = ComponentTransform.GetRotation().Inverse() * InstanceData->SimulatedRotation;
+		RotationOffset.EnforceShortestArcWith(FQuat::Identity);
+
+		FVector OffsetAxis;
+		float OffsetAngle;
+		RotationOffset.ToAxisAndAngle(OffsetAxis, OffsetAngle);
+
+		if (FMath::Abs(OffsetAngle) > MaxAngleRadians)
+		{
+			RotationOffset = FQuat(OffsetAxis, MaxAngleRadians);
+			InstanceData->SimulatedRotation = RotationOffset * ComponentTransform.GetRotation();
+			InstanceData->SimulatedRotation.Normalize();
+		}
+	}
+
+	// Apply the offset adjustments to the simulated transform
+	SimulatedTransform.SetLocation(InstanceData->SimulatedTranslation);
+	SimulatedTransform.SetRotation(InstanceData->SimulatedRotation);
+
+	// Start with the input pose's bone transform, to preserve any adjustments done before this node in the graph
+	FTransform TargetBoneTransform = InputBoneTransform;
+	// Accumulate the simulated transform in, and counter current component transform.
+	TargetBoneTransform.Accumulate(SimulatedTransform * ComponentTransform.Inverse());
+
+	// Offset root bone should not affect scale so take the input
+	TargetBoneTransform.SetScale3D(InputBoneTransform.GetScale3D());
+
+	TargetBoneTransform.NormalizeRotation();
+	check(InstanceData->SimulatedTranslation.ContainsNaN() == false);
+
+	Keyframe->Get()->Pose.LocalTransforms[UE::UAF::FReferencePose::RootBoneIndex] = TargetBoneTransform;
+
+#if ENABLE_VISUAL_LOG
+	if (FVisualLogger::IsRecording())
+	{
+		static const TCHAR* LogName = TEXT("OffsetRootBone");
+		const float InnerCircleRadius = 40.0f;
+		const uint16 CircleThickness = 2;
+		const FVector CircleOffset(0,0,1);
+
+		const FTransform TargetBoneInitialTransformWorld = InputBoneTransform * ComponentTransform;
+		const FTransform TargetBoneTransformWorld = TargetBoneTransform * ComponentTransform;
+		
+		const UObject* LogOwner = InstanceData->HostObject.Get();
+
+		if (SharedData->MaxTranslationError >= 0.0f)
+		{
+			const float OuterCircleRadius = SharedData->MaxTranslationError + InnerCircleRadius;
+			UE_VLOG_CIRCLE_THICK(LogOwner, TEXT("OffsetRootBone"), Display, ComponentTransform.GetLocation() + CircleOffset, FVector::UpVector, OuterCircleRadius, FColor::Red, CircleThickness, TEXT(""));
+		}
+		
+		UE_VLOG_CIRCLE_THICK(LogOwner, LogName, Display, ComponentTransform.GetLocation() + CircleOffset, FVector::UpVector, InnerCircleRadius, FColor::Blue, CircleThickness, TEXT(""));
+		UE_VLOG_ARROW(LogOwner, LogName, Display,
+			ComponentTransform.GetLocation() + CircleOffset,
+			ComponentTransform.GetLocation() + InnerCircleRadius * ComponentTransform.GetRotation().GetRightVector() + CircleOffset,
+			FColor::Blue, TEXT(""));
+		
+		UE_VLOG_CIRCLE_THICK(LogOwner, LogName, Display, TargetBoneTransformWorld.GetLocation() + CircleOffset, FVector::UpVector, InnerCircleRadius, FColor::Green, CircleThickness, TEXT(""));
+		UE_VLOG_ARROW(LogOwner, LogName, Display,
+		 	TargetBoneTransformWorld.GetLocation() + CircleOffset,
+		 	TargetBoneTransformWorld.GetLocation() + InnerCircleRadius * TargetBoneTransformWorld.GetRotation().GetRightVector() + CircleOffset,
+			FColor::Green, TEXT(""));
+	}
+#endif
+}
+
+
diff --git a/Engine/Plugins/Experimental/UAF/UAFWarping/Source/UAFWarping/Private/OffsetRootBoneTrait.h b/Engine/Plugins/Experimental/UAF/UAFWarping/Source/UAFWarping/Private/OffsetRootBoneTrait.h
new file mode 100644
--- /dev/null
+++ b/Engine/Plugins/Experimental/UAF/UAFWarping/Source/UAFWarping/Private/OffsetRootBoneTrait.h
@@ -0,0 +1,196 @@
+// Copyright Epic Games, Inc. All Rights Reserved.
+
+#pragma once
+
+#include "CoreMinimal.h"
+#include "EngineDefines.h"
+#include "TraitCore/Trait.h"
+#include "TraitCore/TraitSharedData.h"
+#include "TraitInterfaces/IAttributeProvider.h"
+#include "TraitInterfaces/IEvaluate.h"
+#include "TraitInterfaces/IHierarchy.h"
+#include "TraitInterfaces/IUpdate.h"
+#include "OffsetRootBoneTrait.generated.h"
+
+UENUM(BlueprintType)
+enum class EUAFOffsetRootBoneMode : uint8
+{
+	// Accumulate the mesh component's movement into the offset.
+	// In this mode, if the mesh component moves 
+	// the offset will counter the motion, and the root will stay in place
+	Accumulate,
+	// Continuously interpolate the offset out
+	// In this mode, if the mesh component moves
+	// The root will stay behind, but will attempt to catch up
+	Interpolate,
+	// Release the offset and stop accumulating the mesh component's movement delta.
+	// In this mode we will "blend out" the offset
+	Release,
+};
+
+
+USTRUCT(meta = (DisplayName = "Offset Root Bone"))
+struct FOffsetRootBoneTraitSharedData : public FAnimNextTraitSharedData
+{
+	GENERATED_BODY()
+
+	// Input pose to be processed
+	UPROPERTY()
+	FAnimNextTraitHandle Input;
+
+	// Current strength of the skeletal control
+	UPROPERTY(EditAnywhere, Category = Alpha, meta = (PinShownByDefault))
+	float Alpha = 1.0f;
+
+	// The translation offset behavior mode
+	UPROPERTY(EditAnywhere, Category = Settings, meta = (FoldProperty))
+	EUAFOffsetRootBoneMode TranslationMode = EUAFOffsetRootBoneMode::Interpolate;
+
+	// The rotation offset behavior mode
+	UPROPERTY(EditAnywhere, Category = Settings, meta = (FoldProperty))
+	EUAFOffsetRootBoneMode RotationMode = EUAFOffsetRootBoneMode::Interpolate;
+
+	// Controls how fast the translation offset is blended out
+	// Values closer to 0 make it faster
+	UPROPERTY(EditAnywhere, Category = Settings, meta = (FoldProperty))
+	float TranslationSmoothingTime = 0.1f;
+
+	// Controls how fast the rotation offset is blended out
+	// Values closer to 0 make it faster
+	UPROPERTY(EditAnywhere, Category = Settings, meta = (FoldProperty))
+	float RotationSmoothingTime = 0.2f;
+
+	// How much the offset can deviate from the mesh component's translation in units
+	// Values lower than 0 disable this limit
+	UPROPERTY(EditAnywhere, Category = Settings, meta = (FoldProperty))
+	float MaxTranslationError = -1.0f;
+
+	// How much the offset can deviate from the mesh component's rotation in degrees
+	// Values lower than 0 disable this limit
+	UPROPERTY(EditAnywhere, Category = Settings, meta = (FoldProperty))
+	float MaxRotationErrorDegrees = -1.0f;
+
+	// Whether to limit the offset's translation interpolation speed to the velocity on the incoming motion
+	// Enabling this prevents the offset sliding when there's little to no translation speed
+	UPROPERTY(EditAnywhere, Category = Settings, meta = (FoldProperty))
+	bool bClampToTranslationVelocity = false;
+
+	// Whether to limit the offset's rotation interpolation speed to the velocity on the incoming motion
+	// Enabling this prevents the offset sliding when there's little to no rotation speed
+	UPROPERTY(EditAnywhere, Category = Settings, meta = (FoldProperty))
+	bool bClampToRotationVelocity = false;
+
+	// How much the offset can blend out, relative to the incoming translation speed
+	// i.e. If root motion is moving at 400cm/s, at 0.5, the offset can blend out at 200cm/s
+	UPROPERTY(EditAnywhere, Category = Settings, meta = (EditCondition = bClampToTranslationVelocity, FoldProperty))
+	float TranslationSpeedRatio = 0.5f;
+
+	// How much the offset can blend out, relative to the incoming rotation speed
+	// i.e. If root motion is rotating at 90 degrees/s, at 0.5, the offset can blend out at 45 degree/s
+	UPROPERTY(EditAnywhere, Category = Settings, meta = (EditCondition = bClampToRotationVelocity, FoldProperty))
+	float RotationSpeedRatio = 0.5f;
+
+	// @TODO Temp / try to remove this. Shouldn't have to feed as argument
+	// This can be the mesh component transform, or the actor. It doesn't really matter since it is
+	// just used to infer the movement of the character in world space from frame to frame
+	// so that the simulated transform can update.
+	UPROPERTY(EditAnywhere, Category = Evaluation, meta = (PinShownByDefault))
+	FTransform HostTransformWorld = FTransform::Identity;
+
+	// Latent pin support boilerplate
+	#define TRAIT_LATENT_PROPERTIES_ENUMERATOR(GeneratorMacro) \
+		GeneratorMacro(Alpha) \
+		GeneratorMacro(HostTransformWorld) \
+		GeneratorMacro(TranslationMode) \
+		GeneratorMacro(RotationMode) \
+		GeneratorMacro(TranslationSmoothingTime) \
+		GeneratorMacro(RotationSmoothingTime) \
+		GeneratorMacro(MaxTranslationError) \
+		GeneratorMacro(MaxRotationErrorDegrees) \
+
+	GENERATE_TRAIT_LATENT_PROPERTIES(FOffsetRootBoneTraitSharedData, TRAIT_LATENT_PROPERTIES_ENUMERATOR)
+	#undef TRAIT_LATENT_PROPERTIES_ENUMERATOR
+};
+
+/**
+ * This trait adds a procedural offset to the root bone coming from the incoming root motion.
+ * It can optionally fade this out over time (controlled by the EUAFOffsetRootBoneMode).
+ *
+ * This allows for a code controlled character to still be influenced by the root motion of the animation,
+ * e.g. during plant turns. It can lead to less foot sliding in a motion matching setup.
+ */
+namespace UE::UAF
+{
+
+struct FOffsetRootBoneTrait : FBaseTrait, IUpdate, IEvaluate, IHierarchy
+{
+	DECLARE_ANIM_TRAIT(FOffsetRootBoneTrait, FBaseTrait)
+
+	using FSharedData = FOffsetRootBoneTraitSharedData;
+
+	struct FInstanceData : FTrait::FInstanceData
+	{
+		FTraitPtr Input;
+
+		/** Delta in seconds between updates, populated during PreUpdate */
+		float DeltaTime = 0.f;
+
+		float Alpha = 1.0f;
+
+		bool bIsFirstUpdate = true;
+
+		FTransform LastMeshComponentTransformWorld = FTransform::Identity;
+		FTransform HostTransformWorld = FTransform::Identity;
+
+		EUAFOffsetRootBoneMode TranslationMode = EUAFOffsetRootBoneMode::Interpolate;
+		EUAFOffsetRootBoneMode RotationMode = EUAFOffsetRootBoneMode::Interpolate;
+
+		float TranslationSmoothingTime = 0.0f;
+		float RotationSmoothingTime = 0.0f;
+		float MaxTranslationError = -1.0f;
+		float MaxRotationErrorDegrees = -1.0f;
+
+		// The simulated world-space transforms for the root bone with offset
+		// Offset = ComponentTransform - SimulatedTransform
+		FVector SimulatedTranslation = FVector::ZeroVector;
+		FQuat SimulatedRotation = FQuat::Identity;
+
+		FVector LastNonZeroRootMotionDirection = FVector::ZeroVector;
+		
+#if ENABLE_ANIM_DEBUG 
+		/** Debug Object for VisualLogger */
+		TObjectPtr<const UObject> HostObject = nullptr;
+#endif // ENABLE_ANIM_DEBUG 
+	};
+
+	// IUpdate impl 
+	virtual void PreUpdate(FUpdateTraversalContext& Context, const TTraitBinding<IUpdate>& Binding, const FTraitUpdateState& TraitState) const override;
+
+	// IEvaluate impl
+	virtual void PostEvaluate(FEvaluateTraversalContext& Context, const TTraitBinding<IEvaluate>& Binding) const override;
+
+	// IHierarchy impl
+	virtual uint32 GetNumChildren(const FExecutionContext& Context, const TTraitBinding<IHierarchy>& Binding) const override;
+	virtual void GetChildren(const FExecutionContext& Context, const TTraitBinding<IHierarchy>& Binding, FChildrenArray& Children) const override;
+};
+
+} // namespace UE::UAF
+
+/** Task to run Offset Root Bone on VM */
+USTRUCT()
+struct FAnimNextOffsetRootBoneTask : public FAnimNextEvaluationTask
+{
+	GENERATED_BODY()
+
+	DECLARE_ANIM_EVALUATION_TASK(FAnimNextOffsetRootBoneTask)
+
+	static FAnimNextOffsetRootBoneTask Make(UE::UAF::FOffsetRootBoneTrait::FInstanceData* InstanceData, const UE::UAF::FOffsetRootBoneTrait::FSharedData* SharedData);
+
+	// Task entry point
+	virtual void Execute(UE::UAF::FEvaluationVM& VM) const override;
+
+private:
+
+	UE::UAF::FOffsetRootBoneTrait::FInstanceData* InstanceData = nullptr;
+	const UE::UAF::FOffsetRootBoneTrait::FSharedData* SharedData = nullptr;
+};
\ No newline at end of file
