Fix compile error in UAF offset root bone trait

diff --git a/Engine/Plugins/Experimental/UAF/UAFWarping/Source/UAFWarping/Private/OffsetRootBoneTrait.cpp b/Engine/Plugins/Experimental/UAF/UAFWarping/Source/UAFWarping/Private/OffsetRootBoneTrait.cpp
--- a/Engine/Plugins/Experimental/UAF/UAFWarping/Source/UAFWarping/Private/OffsetRootBoneTrait.cpp
+++ b/Engine/Plugins/Experimental/UAF/UAFWarping/Source/UAFWarping/Private/OffsetRootBoneTrait.cpp
@@ -9,8 +9,11 @@
 #include "VisualLogger/VisualLogger.h"
 
 #if ENABLE_ANIM_DEBUG
-TAutoConsoleVariable<int32> CVarAnimNodeOffsetRootBoneDebug(TEXT("a.UAF.OffsetRootBone.Debug"), 0, TEXT("Turn on visualization debugging for Offset Root Bone"));
-TAutoConsoleVariable<int32> CVarAnimNodeOffsetRootBoneEnable(TEXT("a.UAF.OffsetRootBone.Enable"), 1, TEXT("Toggle Offset Root Bone"));
+namespace UE::UAF
+{
+	TAutoConsoleVariable<int32> CVarAnimNodeOffsetRootBoneDebug(TEXT("a.UAF.OffsetRootBone.Debug"), 0, TEXT("Turn on visualization debugging for Offset Root Bone"));
+	TAutoConsoleVariable<int32> CVarAnimNodeOffsetRootBoneEnable(TEXT("a.UAF.OffsetRootBone.Enable"), 1, TEXT("Toggle Offset Root Bone"));
+}
 #endif
 
 namespace UE::UAF::OffsetRootBone
