Update variable name in offset root bone
Use debug cvars for offset root bone
Add debug for velocity detected by offset root bone
Add bOnGround option to project simulated translation onto animated ground plane

#virtualized

diff --git a/Engine/Plugins/Experimental/UAF/UAFWarping/Source/UAFWarping/Private/OffsetRootBoneTrait.cpp b/Engine/Plugins/Experimental/UAF/UAFWarping/Source/UAFWarping/Private/OffsetRootBoneTrait.cpp
--- a/Engine/Plugins/Experimental/UAF/UAFWarping/Source/UAFWarping/Private/OffsetRootBoneTrait.cpp
+++ b/Engine/Plugins/Experimental/UAF/UAFWarping/Source/UAFWarping/Private/OffsetRootBoneTrait.cpp
@@ -101,13 +101,15 @@ namespace UE::UAF
 		// Need to copy over latent values from shared data to instance data in order to sample all the pin values (via the getters)
 		// Todo: wish we didn't have to do this
 		InstanceData->Alpha = SharedData->GetAlpha(Binding);
-		InstanceData->HostTransformWorld = SharedData->GetHostTransformWorld(Binding);
+		InstanceData->MeshComponentTransformWorld = SharedData->GetMeshComponentTransformWorld(Binding);
 		InstanceData->TranslationMode = SharedData->GetTranslationMode(Binding);
 		InstanceData->RotationMode = SharedData->GetRotationMode(Binding);
 		InstanceData->MaxTranslationError = SharedData->GetMaxTranslationError(Binding);
 		InstanceData->MaxRotationErrorDegrees = SharedData->GetMaxRotationErrorDegrees(Binding);
 		InstanceData->TranslationSmoothingTime = SharedData->GetTranslationSmoothingTime(Binding);
 		InstanceData->RotationSmoothingTime = SharedData->GetRotationSmoothingTime(Binding);
+		InstanceData->bOnGround = SharedData->GetbOnGround(Binding);
+		InstanceData->AnimatedGroundNormal = SharedData->GetAnimatedGroundNormal(Binding);
 
 #if ENABLE_ANIM_DEBUG 
 		InstanceData->HostObject = Context.GetHostObject();
@@ -173,7 +175,7 @@ void FAnimNextOffsetRootBoneTask::Execute(UE::UAF::FEvaluationVM& VM) const
 
 	const FTransform InputBoneTransform = Keyframe->Get()->Pose.LocalTransforms[UE::UAF::FReferencePose::RootBoneIndex];
 
-	const FTransform ComponentTransform = InstanceData->HostTransformWorld;
+	const FTransform ComponentTransform = InstanceData->MeshComponentTransformWorld;
 
 	if (InstanceData->bIsFirstUpdate)
 	{
@@ -186,7 +188,7 @@ void FAnimNextOffsetRootBoneTask::Execute(UE::UAF::FEvaluationVM& VM) const
 	InstanceData->bIsFirstUpdate = false;
 	
 	// Cycle last transform
-	InstanceData->LastMeshComponentTransformWorld = InstanceData->HostTransformWorld;
+	InstanceData->LastMeshComponentTransformWorld = InstanceData->MeshComponentTransformWorld;
 
 	const EUAFOffsetRootBoneMode CurrentTranslationMode = InstanceData->TranslationMode;
 	const EUAFOffsetRootBoneMode CurrentRotationMode = InstanceData->RotationMode;;
@@ -234,6 +236,11 @@ void FAnimNextOffsetRootBoneTask::Execute(UE::UAF::FEvaluationVM& VM) const
 	InstanceData->SimulatedTranslation = SimulatedTransform.GetLocation();
 	InstanceData->SimulatedRotation = SimulatedTransform.GetRotation();
 
+	if (InstanceData->bOnGround)
+	{
+		InstanceData->SimulatedTranslation = FVector::PointPlaneProject(InstanceData->SimulatedTranslation, ComponentTransform.GetLocation(), InstanceData->AnimatedGroundNormal);
+	}
+
 	if (InstanceData->TranslationMode == EUAFOffsetRootBoneMode::Release ||
 		InstanceData->TranslationMode == EUAFOffsetRootBoneMode::Interpolate)
 	{
@@ -334,7 +341,7 @@ void FAnimNextOffsetRootBoneTask::Execute(UE::UAF::FEvaluationVM& VM) const
 	Keyframe->Get()->Pose.LocalTransforms[UE::UAF::FReferencePose::RootBoneIndex] = TargetBoneTransform;
 
 #if ENABLE_VISUAL_LOG
-	if (FVisualLogger::IsRecording())
+	if (CVarAnimNodeOffsetRootBoneDebug.GetValueOnAnyThread() && FVisualLogger::IsRecording())
 	{
 		static const TCHAR* LogName = TEXT("OffsetRootBone");
 		const float InnerCircleRadius = 40.0f;
diff --git a/Engine/Plugins/Experimental/UAF/UAFWarping/Source/UAFWarping/Private/OffsetRootBoneTrait.h b/Engine/Plugins/Experimental/UAF/UAFWarping/Source/UAFWarping/Private/OffsetRootBoneTrait.h
--- a/Engine/Plugins/Experimental/UAF/UAFWarping/Source/UAFWarping/Private/OffsetRootBoneTrait.h
+++ b/Engine/Plugins/Experimental/UAF/UAFWarping/Source/UAFWarping/Private/OffsetRootBoneTrait.h
@@ -91,22 +91,30 @@ struct FOffsetRootBoneTraitSharedData : public FAnimNextTraitSharedData
 	float RotationSpeedRatio = 0.5f;
 
 	// @TODO Temp / try to remove this. Shouldn't have to feed as argument
-	// This can be the mesh component transform, or the actor. It doesn't really matter since it is
-	// just used to infer the movement of the character in world space from frame to frame
-	// so that the simulated transform can update.
+	// This has to be the mesh component transform in world space
 	UPROPERTY(EditAnywhere, Category = Evaluation, meta = (PinShownByDefault))
-	FTransform HostTransformWorld = FTransform::Identity;
+	FTransform MeshComponentTransformWorld = FTransform::Identity;
+
+	// Signals whether the offset root bone translation should be projected onto a projected ground plane
+	UPROPERTY(EditAnywhere, Category = Settings, meta = (FoldProperty))
+	bool bOnGround = true;
+
+	// Surface normal that the current animation was authored for. Used with the bOnGround flag. 
+	UPROPERTY(EditAnywhere, Category = Settings, meta = (FoldProperty))
+	FVector AnimatedGroundNormal = { 0, 0, 1 };
 
 	// Latent pin support boilerplate
 	#define TRAIT_LATENT_PROPERTIES_ENUMERATOR(GeneratorMacro) \
 		GeneratorMacro(Alpha) \
-		GeneratorMacro(HostTransformWorld) \
+		GeneratorMacro(MeshComponentTransformWorld) \
 		GeneratorMacro(TranslationMode) \
 		GeneratorMacro(RotationMode) \
 		GeneratorMacro(TranslationSmoothingTime) \
 		GeneratorMacro(RotationSmoothingTime) \
 		GeneratorMacro(MaxTranslationError) \
 		GeneratorMacro(MaxRotationErrorDegrees) \
+		GeneratorMacro(bOnGround) \
+		GeneratorMacro(AnimatedGroundNormal) \
 
 	GENERATE_TRAIT_LATENT_PROPERTIES(FOffsetRootBoneTraitSharedData, TRAIT_LATENT_PROPERTIES_ENUMERATOR)
 	#undef TRAIT_LATENT_PROPERTIES_ENUMERATOR
@@ -140,7 +148,7 @@ struct FOffsetRootBoneTrait : FBaseTrait, IUpdate, IEvaluate, IHierarchy
 		bool bIsFirstUpdate = true;
 
 		FTransform LastMeshComponentTransformWorld = FTransform::Identity;
-		FTransform HostTransformWorld = FTransform::Identity;
+		FTransform MeshComponentTransformWorld = FTransform::Identity;
 
 		EUAFOffsetRootBoneMode TranslationMode = EUAFOffsetRootBoneMode::Interpolate;
 		EUAFOffsetRootBoneMode RotationMode = EUAFOffsetRootBoneMode::Interpolate;
@@ -156,6 +164,9 @@ struct FOffsetRootBoneTrait : FBaseTrait, IUpdate, IEvaluate, IHierarchy
 		FQuat SimulatedRotation = FQuat::Identity;
 
 		FVector LastNonZeroRootMotionDirection = FVector::ZeroVector;
+
+		bool bOnGround = true;
+		FVector AnimatedGroundNormal = FVector::ZeroVector;
 		
 #if ENABLE_ANIM_DEBUG 
 		/** Debug Object for VisualLogger */
