Add inline meta to variables that we don't expect to change dynamically
Remove unnecessary FoldProperty flag

diff --git a/Engine/Plugins/Experimental/UAF/UAFWarping/Source/UAFWarping/Private/OffsetRootBoneTrait.h b/Engine/Plugins/Experimental/UAF/UAFWarping/Source/UAFWarping/Private/OffsetRootBoneTrait.h
--- a/Engine/Plugins/Experimental/UAF/UAFWarping/Source/UAFWarping/Private/OffsetRootBoneTrait.h
+++ b/Engine/Plugins/Experimental/UAF/UAFWarping/Source/UAFWarping/Private/OffsetRootBoneTrait.h
@@ -43,51 +43,51 @@ struct FOffsetRootBoneTraitSharedData : public FAnimNextTraitSharedData
 	float Alpha = 1.0f;
 
 	// The translation offset behavior mode
-	UPROPERTY(EditAnywhere, Category = Settings, meta = (FoldProperty))
+	UPROPERTY(EditAnywhere, Category = Settings)
 	EUAFOffsetRootBoneMode TranslationMode = EUAFOffsetRootBoneMode::Interpolate;
 
 	// The rotation offset behavior mode
-	UPROPERTY(EditAnywhere, Category = Settings, meta = (FoldProperty))
+	UPROPERTY(EditAnywhere, Category = Settings)
 	EUAFOffsetRootBoneMode RotationMode = EUAFOffsetRootBoneMode::Interpolate;
 
 	// Controls how fast the translation offset is blended out
 	// Values closer to 0 make it faster
-	UPROPERTY(EditAnywhere, Category = Settings, meta = (FoldProperty))
+	UPROPERTY(EditAnywhere, Category = Settings)
 	float TranslationSmoothingTime = 0.1f;
 
 	// Controls how fast the rotation offset is blended out
 	// Values closer to 0 make it faster
-	UPROPERTY(EditAnywhere, Category = Settings, meta = (FoldProperty))
+	UPROPERTY(EditAnywhere, Category = Settings)
 	float RotationSmoothingTime = 0.2f;
 
 	// How much the offset can deviate from the mesh component's translation in units
 	// Values lower than 0 disable this limit
-	UPROPERTY(EditAnywhere, Category = Settings, meta = (FoldProperty))
+	UPROPERTY(EditAnywhere, Category = Settings)
 	float MaxTranslationError = -1.0f;
 
 	// How much the offset can deviate from the mesh component's rotation in degrees
 	// Values lower than 0 disable this limit
-	UPROPERTY(EditAnywhere, Category = Settings, meta = (FoldProperty))
+	UPROPERTY(EditAnywhere, Category = Settings)
 	float MaxRotationErrorDegrees = -1.0f;
 
 	// Whether to limit the offset's translation interpolation speed to the velocity on the incoming motion
 	// Enabling this prevents the offset sliding when there's little to no translation speed
-	UPROPERTY(EditAnywhere, Category = Settings, meta = (FoldProperty))
+	UPROPERTY(EditAnywhere, Category = Settings, meta=(Inline))
 	bool bClampToTranslationVelocity = false;
 
 	// Whether to limit the offset's rotation interpolation speed to the velocity on the incoming motion
 	// Enabling this prevents the offset sliding when there's little to no rotation speed
-	UPROPERTY(EditAnywhere, Category = Settings, meta = (FoldProperty))
+	UPROPERTY(EditAnywhere, Category = Settings, meta=(Inline))
 	bool bClampToRotationVelocity = false;
 
 	// How much the offset can blend out, relative to the incoming translation speed
 	// i.e. If root motion is moving at 400cm/s, at 0.5, the offset can blend out at 200cm/s
-	UPROPERTY(EditAnywhere, Category = Settings, meta = (EditCondition = bClampToTranslationVelocity, FoldProperty))
+	UPROPERTY(EditAnywhere, Category = Settings, meta = (EditCondition = bClampToTranslationVelocity, Inline))
 	float TranslationSpeedRatio = 0.5f;
 
 	// How much the offset can blend out, relative to the incoming rotation speed
 	// i.e. If root motion is rotating at 90 degrees/s, at 0.5, the offset can blend out at 45 degree/s
-	UPROPERTY(EditAnywhere, Category = Settings, meta = (EditCondition = bClampToRotationVelocity, FoldProperty))
+	UPROPERTY(EditAnywhere, Category = Settings, meta = (EditCondition = bClampToRotationVelocity, Inline))
 	float RotationSpeedRatio = 0.5f;
 
 	// @TODO Temp / try to remove this. Shouldn't have to feed as argument
@@ -96,11 +96,11 @@ struct FOffsetRootBoneTraitSharedData : public FAnimNextTraitSharedData
 	FTransform MeshComponentTransformWorld = FTransform::Identity;
 
 	// Signals whether the offset root bone translation should be projected onto a projected ground plane
-	UPROPERTY(EditAnywhere, Category = Settings, meta = (FoldProperty))
+	UPROPERTY(EditAnywhere, Category = Settings)
 	bool bOnGround = true;
 
 	// Surface normal that the current animation was authored for. Used with the bOnGround flag. 
-	UPROPERTY(EditAnywhere, Category = Settings, meta = (FoldProperty))
+	UPROPERTY(EditAnywhere, Category = Settings)
 	FVector AnimatedGroundNormal = { 0, 0, 1 };
 
 	// Latent pin support boilerplate
