Add new trait to override root motion

diff --git a/Engine/Plugins/Experimental/UAF/UAFWarping/Source/UAFWarping/Private/OverrideRootMotionTrait.cpp b/Engine/Plugins/Experimental/UAF/UAFWarping/Source/UAFWarping/Private/OverrideRootMotionTrait.cpp
new file mode 100644
--- /dev/null
+++ b/Engine/Plugins/Experimental/UAF/UAFWarping/Source/UAFWarping/Private/OverrideRootMotionTrait.cpp
@@ -0,0 +1,119 @@
+// Copyright Epic Games, Inc. All Rights Reserved.
+
+#include "OverrideRootMotionTrait.h"
+
+#include "AnimNextWarpingLog.h"
+#include "Animation/AnimRootMotionProvider.h"
+#include "Animation/SpringMath.h"
+#include "EvaluationVM/EvaluationVM.h"
+#include "VisualLogger/VisualLogger.h"
+
+////////////////////////////////////////////////////////////////////////////////////////////////////////////////////////
+// FOverrideRootMotionTrait
+
+namespace UE::UAF
+{
+	AUTO_REGISTER_ANIM_TRAIT(FOverrideRootMotionTrait)
+
+	#define TRAIT_INTERFACE_ENUMERATOR(GeneratorMacro) \
+		GeneratorMacro(IUpdate) \
+		GeneratorMacro(IEvaluate) \
+
+	// Trait implementation boilerplate
+	GENERATE_ANIM_TRAIT_IMPLEMENTATION(FOverrideRootMotionTrait, TRAIT_INTERFACE_ENUMERATOR, NULL_ANIM_TRAIT_INTERFACE_ENUMERATOR, NULL_ANIM_TRAIT_EVENT_ENUMERATOR)
+	#undef TRAIT_INTERFACE_ENUMERATOR
+
+	void FOverrideRootMotionTrait::PreUpdate(FUpdateTraversalContext& Context, const TTraitBinding<IUpdate>& Binding, const FTraitUpdateState& TraitState) const
+	{
+		FInstanceData* InstanceData = Binding.GetInstanceData<FInstanceData>();
+		check(InstanceData);
+
+		const FSharedData* SharedData = Binding.GetSharedData<FSharedData>();
+		check(SharedData);
+
+		IUpdate::PreUpdate(Context, Binding, TraitState);
+	}
+
+	void FOverrideRootMotionTrait::PostEvaluate(FEvaluateTraversalContext& Context, const TTraitBinding<IEvaluate>& Binding) const
+	{
+		IEvaluate::PostEvaluate(Context, Binding);
+
+		const FSharedData* SharedData = Binding.GetSharedData<FSharedData>();
+		check(SharedData);
+
+		Context.AppendTask(FAnimNextOverrideRootMotionTask::Make(SharedData->GetAlpha(Binding), SharedData->GetOverrideRootMotionDelta(Binding), SharedData->GetOverrideRootMotionMode(Binding)));
+	}
+
+} // namespace UE::UAF
+
+
+////////////////////////////////////////////////////////////////////////////////////////////////////////////////////////
+// FAnimNextOverrideRootMotionTask
+
+FAnimNextOverrideRootMotionTask FAnimNextOverrideRootMotionTask::Make(float Alpha, FTransform OverrideRootMotionDelta, EUAFOverrideRootMotionMode OverrideRootMotionMode)
+{
+	FAnimNextOverrideRootMotionTask Task;
+	Task.Alpha = Alpha;
+	Task.OverrideRootMotionDelta = OverrideRootMotionDelta;
+	Task.OverrideRootMotionMode = OverrideRootMotionMode;
+	return Task;
+}
+
+void FAnimNextOverrideRootMotionTask::Execute(UE::UAF::FEvaluationVM& VM) const
+{
+	using namespace UE::UAF;
+	QUICK_SCOPE_CYCLE_COUNTER(FAnimNextOffsetRootBoneTask_Execute);
+
+	if (Alpha < UE_SMALL_NUMBER)
+	{
+		return;
+	}
+
+	// Extract the root motion for this frame
+	TUniquePtr<FKeyframeState>* Keyframe = VM.PeekValueMutable<TUniquePtr<FKeyframeState>>(KEYFRAME_STACK_NAME, 0);
+	if (Keyframe == nullptr || Keyframe->Get() == nullptr)
+	{
+		return;
+	}
+	
+	FTransform RootMotionTransformDelta = FTransform::Identity;
+	const UE::Anim::IAnimRootMotionProvider* RootMotionProvider = UE::Anim::IAnimRootMotionProvider::Get();
+	if (!RootMotionProvider)
+	{
+		UE_LOG(LogAnimNextWarping, Error, TEXT("FAnimNextOverrideRootMotionTask::Execute, missing RootMotionProvider"));
+		return;
+	}
+
+	bool RequiresBlend = Alpha < 1.0f - UE_SMALL_NUMBER;
+	// We only need to check the incoming root motion if we are going to blend it
+	// Or add it to the override
+	if (RequiresBlend || OverrideRootMotionMode == EUAFOverrideRootMotionMode::Additive)
+	{
+		// Note: this can fail, but if so we just use identity as the incoming root motion delta
+		RootMotionProvider->ExtractRootMotion(Keyframe->Get()->Attributes, RootMotionTransformDelta);
+	}
+
+	FTransform NewRootMotionToApply;
+
+	if (OverrideRootMotionMode == EUAFOverrideRootMotionMode::Replace)
+	{
+		NewRootMotionToApply = OverrideRootMotionDelta;
+	}
+	else
+	{
+		NewRootMotionToApply = RootMotionTransformDelta; 
+		NewRootMotionToApply.NormalizeRotation();
+		NewRootMotionToApply.Accumulate(OverrideRootMotionDelta);
+	}
+
+	if (RequiresBlend)
+	{
+		// Need to blend effect on
+		NewRootMotionToApply.Blend(RootMotionTransformDelta, NewRootMotionToApply, Alpha);
+	}
+
+	RootMotionProvider->SetRootMotion(NewRootMotionToApply, Keyframe->Get()->Attributes);
+	
+}
+
+
diff --git a/Engine/Plugins/Experimental/UAF/UAFWarping/Source/UAFWarping/Private/OverrideRootMotionTrait.h b/Engine/Plugins/Experimental/UAF/UAFWarping/Source/UAFWarping/Private/OverrideRootMotionTrait.h
new file mode 100644
--- /dev/null
+++ b/Engine/Plugins/Experimental/UAF/UAFWarping/Source/UAFWarping/Private/OverrideRootMotionTrait.h
@@ -0,0 +1,92 @@
+// Copyright Epic Games, Inc. All Rights Reserved.
+
+#pragma once
+
+#include "CoreMinimal.h"
+#include "EngineDefines.h"
+#include "TraitCore/Trait.h"
+#include "TraitCore/TraitSharedData.h"
+#include "TraitInterfaces/IAttributeProvider.h"
+#include "TraitInterfaces/IEvaluate.h"
+#include "TraitInterfaces/IUpdate.h"
+#include "OverrideRootMotionTrait.generated.h"
+
+UENUM(BlueprintType)
+enum class EUAFOverrideRootMotionMode : uint8
+{
+	Replace,
+	Additive
+};
+
+
+USTRUCT(meta = (DisplayName = "Override Root Motion"))
+struct FOverrideRootMotionTraitSharedData : public FAnimNextTraitSharedData
+{
+	GENERATED_BODY()
+
+	// Current strength of the override
+	UPROPERTY(EditAnywhere, Category = Settings)
+	float Alpha = 1.0f;
+
+	UPROPERTY(EditAnywhere, Category = Settings)
+	FTransform OverrideRootMotionDelta = FTransform::Identity;
+
+	UPROPERTY(EditAnywhere, Category = Settings)
+	EUAFOverrideRootMotionMode OverrideRootMotionMode = EUAFOverrideRootMotionMode::Replace;
+
+	// Latent pin support boilerplate
+	#define TRAIT_LATENT_PROPERTIES_ENUMERATOR(GeneratorMacro) \
+		GeneratorMacro(Alpha) \
+		GeneratorMacro(OverrideRootMotionDelta) \
+		GeneratorMacro(OverrideRootMotionMode) \
+
+	GENERATE_TRAIT_LATENT_PROPERTIES(FOverrideRootMotionTraitSharedData, TRAIT_LATENT_PROPERTIES_ENUMERATOR)
+	#undef TRAIT_LATENT_PROPERTIES_ENUMERATOR
+};
+
+namespace UE::UAF
+{
+
+/**
+ * Override the root motion attributes with custom values.
+ * 
+ * Ex: allows code driven motion to be injected into the animation system as root motion.
+ */
+struct FOverrideRootMotionTrait : FAdditiveTrait, IUpdate, IEvaluate
+{
+	DECLARE_ANIM_TRAIT(FOverrideRootMotionTrait, FAdditiveTrait)
+
+	using FSharedData = FOverrideRootMotionTraitSharedData;
+
+	struct FInstanceData : FTrait::FInstanceData
+	{
+	};
+
+	// IUpdate impl 
+	virtual void PreUpdate(FUpdateTraversalContext& Context, const TTraitBinding<IUpdate>& Binding, const FTraitUpdateState& TraitState) const override;
+
+	// IEvaluate impl
+	virtual void PostEvaluate(FEvaluateTraversalContext& Context, const TTraitBinding<IEvaluate>& Binding) const override;
+};
+
+} // namespace UE::UAF
+
+/** Task to run override root motion on VM */
+USTRUCT()
+struct FAnimNextOverrideRootMotionTask : public FAnimNextEvaluationTask
+{
+	GENERATED_BODY()
+
+	DECLARE_ANIM_EVALUATION_TASK(FAnimNextOverrideRootMotionTask)
+
+	static FAnimNextOverrideRootMotionTask Make(float Alpha, FTransform OverrideRootMotionDelta, EUAFOverrideRootMotionMode OverrideRootMotionMode);
+
+	// Task entry point
+	virtual void Execute(UE::UAF::FEvaluationVM& VM) const override;
+
+private:
+
+	float Alpha = 1.0f;
+	FTransform OverrideRootMotionDelta = FTransform::Identity;
+	EUAFOverrideRootMotionMode OverrideRootMotionMode = EUAFOverrideRootMotionMode::Replace;
+};
\ No newline at end of file
