UAF: Fixed small latent pin issue in OffsetRootBoneTrait

#Synced-CL 48507157

diff --git a/Engine/Plugins/Experimental/UAF/UAFWarping/Source/UAFWarping/Private/OffsetRootBoneTrait.h b/Engine/Plugins/Experimental/UAF/UAFWarping/Source/UAFWarping/Private/OffsetRootBoneTrait.h
--- a/Engine/Plugins/Experimental/UAF/UAFWarping/Source/UAFWarping/Private/OffsetRootBoneTrait.h
+++ b/Engine/Plugins/Experimental/UAF/UAFWarping/Source/UAFWarping/Private/OffsetRootBoneTrait.h
@@ -104,11 +104,11 @@ struct FOffsetRootBoneTraitSharedData : public FAnimNextTraitSharedData
 	FVector AnimatedGroundNormal = { 0, 0, 1 };
 
 	// Distance threshold used to detect teleports
-	UPROPERTY(EditAnywhere, Category = Settings)
+	UPROPERTY(EditAnywhere, Category = Settings, meta = (Inline))
 	float TeleportDistanceThreshold = 300.0f;
 
 	// If to reset the root bone offset on teleport
-	UPROPERTY(EditAnywhere, Category = Settings)
+	UPROPERTY(EditAnywhere, Category = Settings, meta = (Inline))
 	bool bResetOnTeleport = false;
 
 	// Latent pin support boilerplate
