UAF: updates to strafe warping trait
* Fix for bad pole vector calculation on certain animations
* strafe warping trait now sets the IK target bone correctly, for use with other IK solvers downstream (e.g. control rig)
* added option to skip the basic IK done inside the strafe warping node

#Synced-CL 49767845

diff --git a/Engine/Plugins/Experimental/UAF/UAFWarping/Source/UAFWarping/Private/StrafeWarpingTrait.cpp b/Engine/Plugins/Experimental/UAF/UAFWarping/Source/UAFWarping/Private/StrafeWarpingTrait.cpp
--- a/Engine/Plugins/Experimental/UAF/UAFWarping/Source/UAFWarping/Private/StrafeWarpingTrait.cpp
+++ b/Engine/Plugins/Experimental/UAF/UAFWarping/Source/UAFWarping/Private/StrafeWarpingTrait.cpp
@@ -451,7 +451,6 @@ void FAnimNextStrafeWarpingTask::DoWarpPose(TUniquePtr<UE::UAF::FKeyframeState>&
 	{
 		{
 			const FQuat IKRootBoneRotation = FQuat(RotationAxisVector, InstanceData->OrientationAngleForPoseWarpRad * IKFootRootOrientationAlpha);
-
 			// IK Feet 
 			// We want these to keep their original component space orientation
 			// But we want them to translate based on some rotation
@@ -462,72 +461,120 @@ void FAnimNextStrafeWarpingTask::DoWarpPose(TUniquePtr<UE::UAF::FKeyframeState>&
 				for (int32 ArrayIndex = 0; ArrayIndex < NumIKFootBones; ArrayIndex++)
 				{
 					const FStrafeWarpFootData& FootData = SharedData->FootData[ArrayIndex];
-					
-					FBoneIndexType LegRootIndex = Pose.FindLODBoneIndexFromBoneName(FootData.LegRoot);
-					FBoneIndexType LegMidIndex = Pose.FindLODBoneIndexFromBoneName(FootData.LegMid);
 					FBoneIndexType LegTipIndex = Pose.FindLODBoneIndexFromBoneName(FootData.LegTip);
-
-					// Validate data
-					if (LegRootIndex == INDEX_NONE || LegMidIndex == INDEX_NONE || LegTipIndex == INDEX_NONE)
+					if (LegTipIndex == INDEX_NONE)
 					{
 						continue;
 					}
 
-					// Todo: should these be ensure? error? how to handle?
-					check(Pose.GetLODBoneParentIndex(LegTipIndex) == LegMidIndex);
-					check(Pose.GetLODBoneParentIndex(LegMidIndex) == LegRootIndex);
-
-					FTransform LegRootTransformCS(ComponentSpaceTransforms.Rotations[LegRootIndex], ComponentSpaceTransforms.Translations[LegRootIndex]);
-					FTransform LegMidTransformCS(ComponentSpaceTransforms.Rotations[LegMidIndex], ComponentSpaceTransforms.Translations[LegMidIndex]);
-					FTransform LegTipTransformCS(ComponentSpaceTransforms.Rotations[LegTipIndex], ComponentSpaceTransforms.Translations[LegTipIndex]);
+					FTransform FootTargetTransform(ComponentSpaceTransforms.Rotations[LegTipIndex], ComponentSpaceTransforms.Translations[LegTipIndex]);
+					FootTargetTransform.SetLocation( IKRootBoneRotation.RotateVector(FootTargetTransform.GetLocation()));
 
-					FVector FootTargetPosition = IKRootBoneRotation.RotateVector(LegTipTransformCS.GetLocation());
-
-					// Joint (knee) target vector, needs to be driven from animation and define a good solving plan
-					FVector JointTarget = FootTargetPosition; // TODO This might break if leg is straight 
-
-					AnimationCore::SolveTwoBoneIK(LegRootTransformCS,
-						LegMidTransformCS,
-						LegTipTransformCS,
-						JointTarget,
-						FootTargetPosition,
-						false,
-						0.0,
-						0.0
-						);
+#if ENABLE_VISUAL_LOG && ENABLE_ANIM_DEBUG
+					if (FVisualLogger::IsRecording())
+					{
+						const float DebugDrawSpehereRadius = 10.f;
+						UE_VLOG_SPHERE(InstanceData->HostObject, "OrientationWarping", Display, InstanceData->RootBoneTransform.TransformPosition(FootTargetTransform.GetLocation()), DebugDrawSpehereRadius, FColor::Green, TEXT(""));
+					}
+#endif // ENABLE_VISUAL_LOG && ENABLE_ANIM_DEBUG
 
-					// Apply results
+					if (SharedData->FootMode == EStrafeWarpingFootMode::DoBasicIK)
+					{
+						FBoneIndexType LegRootIndex = Pose.FindLODBoneIndexFromBoneName(FootData.LegRoot);
+						FBoneIndexType LegMidIndex = Pose.FindLODBoneIndexFromBoneName(FootData.LegMid);
 
-					// Be careful reverse TTransform vs FQuat order of operations
-					FTransform LocalTip = LegTipTransformCS * LegMidTransformCS.Inverse();
-					FTransform LocalMid = LegMidTransformCS * LegRootTransformCS.Inverse();
+						// Validate data
+						if (LegRootIndex == INDEX_NONE || LegMidIndex == INDEX_NONE)
+						{
+							continue;
+						}
 
-					// Figure out local root by applying difference in CS
-					FQuat LocalRootRotationDiff = ComponentSpaceTransforms[LegRootIndex].Rotation.Inverse() * LegRootTransformCS.GetRotation(); 
+						check(Pose.GetLODBoneParentIndex(LegTipIndex) == LegMidIndex);
+						check(Pose.GetLODBoneParentIndex(LegMidIndex) == LegRootIndex);
+
+						FTransform LegRootTransformCS(ComponentSpaceTransforms.Rotations[LegRootIndex], ComponentSpaceTransforms.Translations[LegRootIndex]);
+						FTransform LegMidTransformCS(ComponentSpaceTransforms.Rotations[LegMidIndex], ComponentSpaceTransforms.Translations[LegMidIndex]);
+						FTransform LegTipTransformCS(ComponentSpaceTransforms.Rotations[LegTipIndex], ComponentSpaceTransforms.Translations[LegTipIndex]);
+
+						// Calculate pole vector
+						// By using the side direction of the knee and the target, we can ensure the pole target is always 90 degrees from the target direction
+						FVector PoleTarget = FootTargetTransform.GetLocation() - LegRootTransformCS.GetLocation();
+						FVector SideVectorCS = LegMidTransformCS.GetRotation().RotateVector(FootData.JointSideVector);
+						PoleTarget = PoleTarget.Cross(SideVectorCS);
+
+						// Pole vector is expected to be provided as a location in AnimationCore::SolveTwoBoneIK
+						// so convert from direction to location
+						PoleTarget += LegRootTransformCS.GetLocation();
+
+						AnimationCore::SolveTwoBoneIK(LegRootTransformCS,
+							LegMidTransformCS,
+							LegTipTransformCS,
+							PoleTarget,
+							FootTargetTransform.GetLocation(),
+							false,
+							0.0,
+							0.0
+							);
+
+						LegTipTransformCS.SetRotation(FootTargetTransform.GetRotation());
+
+						// Apply results
+
+						// Be careful reverse TTransform vs FQuat order of operations
+						FTransform LocalTip = LegTipTransformCS * LegMidTransformCS.Inverse();
+						FTransform LocalMid = LegMidTransformCS * LegRootTransformCS.Inverse();
+
+						// Figure out local root by applying difference in CS. We can't use the parent component space since it hasn't been updated (unlike the foot and knee)
+						// This avoids computing all the component space transforms again
+						// We can set the local transform of the root by applying the difference between the component space root transform
+						FQuat LocalRootRotationDiff = ComponentSpaceTransforms[LegRootIndex].Rotation.Inverse() * LegRootTransformCS.GetRotation(); 
+
+						PoseTransforms[LegTipIndex].Translation = LocalTip.GetLocation();
+						PoseTransforms[LegTipIndex].Rotation = LocalTip.GetRotation();
+						PoseTransforms[LegMidIndex].Translation = LocalMid.GetLocation();
+						PoseTransforms[LegMidIndex].Rotation = LocalMid.GetRotation();
+						PoseTransforms[LegRootIndex].Rotation = PoseTransforms[LegRootIndex].Rotation * LocalRootRotationDiff;
+						
+
+	#if ENABLE_VISUAL_LOG && ENABLE_ANIM_DEBUG
+						if (FVisualLogger::IsRecording())
+						{
+							const float DebugDrawScale = 10.0f;
+							const uint16 DebugDrawThickness = 1;
+							const float DebugDrawSpehereRadius = 10.f;
+							// Careful: FTransform multiplication is reveresed from quat!
+							FTransform FootWorldTransformOriginal = FTransform(ComponentSpaceTransforms.Rotations[LegTipIndex], ComponentSpaceTransforms.Translations[LegTipIndex]) * InstanceData->RootBoneTransform;
+							FTransform FootWorldTransformTarget = FTransform(ComponentSpaceTransforms.Rotations[LegTipIndex], FootTargetTransform.GetLocation()) * InstanceData->RootBoneTransform;
+							FTransform FootWorldTransformSolved = LegTipTransformCS * InstanceData->RootBoneTransform;
+							UE_VLOG_COORDINATESYSTEM(InstanceData->HostObject, "OrientationWarping", Display, FootWorldTransformTarget.GetLocation(), FootWorldTransformTarget.GetRotation().Rotator(), DebugDrawScale, FColor::Green, DebugDrawThickness, TEXT(""));
+							UE_VLOG_COORDINATESYSTEM(InstanceData->HostObject, "OrientationWarping", Display, FootWorldTransformOriginal.GetLocation(), FootWorldTransformOriginal.GetRotation().Rotator(), DebugDrawScale, FColor::Green, DebugDrawThickness, TEXT(""));
+							UE_VLOG_ARROW(InstanceData->HostObject, "OrientationWarping", Display, FootWorldTransformOriginal.GetLocation(), FootWorldTransformTarget.GetLocation(), FColor::Magenta, TEXT(""));
+							UE_VLOG_ARROW(InstanceData->HostObject, "OrientationWarping", Display, LegMidTransformCS.GetLocation() + InstanceData->RootBoneTransform.GetLocation(), PoleTarget + InstanceData->RootBoneTransform.GetLocation(), FColor::Red, TEXT(""));
+							UE_VLOG_SPHERE(InstanceData->HostObject, "OrientationWarping", Display, FootWorldTransformSolved.GetTranslation(), DebugDrawSpehereRadius, FColor::Yellow, TEXT(""));
+						}
+	#endif // ENABLE_VISUAL_LOG && ENABLE_ANIM_DEBUG
+					}
 
-					PoseTransforms[LegTipIndex].Translation = LocalTip.GetLocation();
-					PoseTransforms[LegTipIndex].Rotation = LocalTip.GetRotation();
-					PoseTransforms[LegMidIndex].Translation = LocalMid.GetLocation();
-					PoseTransforms[LegMidIndex].Rotation = LocalMid.GetRotation();
-					PoseTransforms[LegRootIndex].Rotation = PoseTransforms[LegRootIndex].Rotation * LocalRootRotationDiff;
 					
+					// Set the IK bone to match the foot target
+					// Better to do this even if we also did basic IK, in case other nodes later down the chain rely on the ik target bone
 
-#if ENABLE_VISUAL_LOG && ENABLE_ANIM_DEBUG
-					if (FVisualLogger::IsRecording())
+					FBoneIndexType IKTargetIndex = Pose.FindLODBoneIndexFromBoneName(FootData.IKTargetBone);
+
+					if (IKTargetIndex != INDEX_NONE)
 					{
-						const float DebugDrawScale = 10.0f;
-						const uint16 DebugDrawThickness = 1;
-						const float DebugDrawSpehereRadius = 10.f;
-						// Careful: FTransform multiplication is reveresed from quat!
-						FTransform FootWorldTransformOriginal = FTransform(ComponentSpaceTransforms.Rotations[LegTipIndex], ComponentSpaceTransforms.Translations[LegTipIndex]) * InstanceData->RootBoneTransform;
-						FTransform FootWorldTransformTarget = FTransform(ComponentSpaceTransforms.Rotations[LegTipIndex], FootTargetPosition) * InstanceData->RootBoneTransform;
-						FTransform FootWorldTransformSolved = LegTipTransformCS * InstanceData->RootBoneTransform;
-						UE_VLOG_COORDINATESYSTEM(InstanceData->HostObject, "OrientationWarping", Display, FootWorldTransformTarget.GetLocation(), FootWorldTransformTarget.GetRotation().Rotator(), DebugDrawScale, FColor::Green, DebugDrawThickness, TEXT(""));
-						UE_VLOG_COORDINATESYSTEM(InstanceData->HostObject, "OrientationWarping", Display, FootWorldTransformOriginal.GetLocation(), FootWorldTransformOriginal.GetRotation().Rotator(), DebugDrawScale, FColor::Green, DebugDrawThickness, TEXT(""));
-						UE_VLOG_ARROW(InstanceData->HostObject, "OrientationWarping", Display, FootWorldTransformOriginal.GetLocation(), FootWorldTransformTarget.GetLocation(), FColor::Magenta, TEXT(""));
-						UE_VLOG_SPHERE(InstanceData->HostObject, "OrientationWarping", Display, FootWorldTransformSolved.GetTranslation(), DebugDrawSpehereRadius, FColor::Yellow, TEXT(""));
+						// Set IK target to equal leg tip in component space
+						int IKTargetParent = Pose.GetLODBoneParentIndex(IKTargetIndex);		// Usually the root bone but we can make it work for any parent
+						if (IKTargetParent != INDEX_NONE)
+						{
+							FTransform IKTargetCS = FootTargetTransform;
+							FTransform IKTargetParentCS = ComponentSpaceTransforms[IKTargetParent];
+							FTransform IKTargetLocal = IKTargetCS * IKTargetParentCS.Inverse();
+
+							PoseTransforms[IKTargetIndex] = IKTargetLocal;
+						}
 					}
-#endif // ENABLE_VISUAL_LOG && ENABLE_ANIM_DEBUG
+				
 				}
 			}
 		}
diff --git a/Engine/Plugins/Experimental/UAF/UAFWarping/Source/UAFWarping/Private/StrafeWarpingTrait.h b/Engine/Plugins/Experimental/UAF/UAFWarping/Source/UAFWarping/Private/StrafeWarpingTrait.h
--- a/Engine/Plugins/Experimental/UAF/UAFWarping/Source/UAFWarping/Private/StrafeWarpingTrait.h
+++ b/Engine/Plugins/Experimental/UAF/UAFWarping/Source/UAFWarping/Private/StrafeWarpingTrait.h
@@ -11,23 +11,44 @@
 #include "TraitInterfaces/IUpdate.h"
 #include "StrafeWarpingTrait.generated.h"
 
+UENUM(BlueprintType)
+enum class EStrafeWarpingFootMode : uint8
+{
+	DoBasicIK,
+	SetIKBoneOnly,
+};
+
 // Small structure to define a simple IK setup for legs
 USTRUCT()
 struct FStrafeWarpFootData
 {
 	GENERATED_BODY()
 
+	// Used for basic IK performed on the node itself
+	// Only required if EStrafeWarpingFootMode is set to `DoBasicIK`
 	// e.g. the thigh bone
 	UPROPERTY(EditAnywhere, Category=Settings)
 	FName LegRoot;
 
+	// Used for basic IK performed on the node itself
+	// Only required if EStrafeWarpingFootMode is set to `DoBasicIK`
 	// e.g. the knee bone
 	UPROPERTY(EditAnywhere, Category=Settings)
 	FName LegMid;
 
+	// This bone is required in all modes
 	// e.g. the foot bone
 	UPROPERTY(EditAnywhere, Category=Settings)
 	FName LegTip;
+
+	// Set this bone to the IK target (in character space)
+	// Can be invalid
+	UPROPERTY(EditAnywhere, Category=Settings)
+	FName IKTargetBone;
+
+	// Used for pole vector calculation. Should be the side of the knee in its local space
+	UPROPERTY(EditAnywhere, Category=Settings)
+	FVector JointSideVector = FVector::ZAxisVector;
 };
 
 USTRUCT(meta = (DisplayName = "Strafe Warping"))
@@ -112,6 +133,9 @@ struct FStrafeWarpingTraitSharedData : public FAnimNextTraitSharedData
 	UPROPERTY(EditAnywhere, Category=Settings, meta=(DisplayName="IK Foot Bones"), meta=(Inline, Hidden))
 	TArray<FStrafeWarpFootData> FootData;
 
+	UPROPERTY(EditAnywhere, Category=Settings, meta=(Inline))
+	EStrafeWarpingFootMode FootMode = EStrafeWarpingFootMode::DoBasicIK;
+
 	// Latent pin support boilerplate
 	#define TRAIT_LATENT_PROPERTIES_ENUMERATOR(GeneratorMacro) \
 		GeneratorMacro(Alpha) \
