Remove cvar to control offset root bone trait visual log, we should do it all the time the visual logger is recording

#Synced-CL 50050531

diff --git a/Engine/Plugins/Experimental/UAF/UAFWarping/Source/UAFWarping/Private/OffsetRootBoneTrait.cpp b/Engine/Plugins/Experimental/UAF/UAFWarping/Source/UAFWarping/Private/OffsetRootBoneTrait.cpp
--- a/Engine/Plugins/Experimental/UAF/UAFWarping/Source/UAFWarping/Private/OffsetRootBoneTrait.cpp
+++ b/Engine/Plugins/Experimental/UAF/UAFWarping/Source/UAFWarping/Private/OffsetRootBoneTrait.cpp
@@ -11,7 +11,6 @@
 #if ENABLE_ANIM_DEBUG
 namespace UE::UAF
 {
-	TAutoConsoleVariable<int32> CVarAnimNodeOffsetRootBoneDebug(TEXT("a.UAF.OffsetRootBone.Debug"), 0, TEXT("Turn on visualization debugging for Offset Root Bone"));
 	TAutoConsoleVariable<int32> CVarAnimNodeOffsetRootBoneEnable(TEXT("a.UAF.OffsetRootBone.Enable"), 1, TEXT("Toggle Offset Root Bone"));
 }
 #endif
@@ -364,7 +363,7 @@ void FAnimNextOffsetRootBoneTask::Execute(UE::UAF::FEvaluationVM& VM) const
 	Keyframe->Get()->Pose.LocalTransforms[UE::UAF::FReferencePose::RootBoneIndex] = TargetBoneTransform;
 
 #if ENABLE_VISUAL_LOG && ENABLE_ANIM_DEBUG
-	if (CVarAnimNodeOffsetRootBoneDebug.GetValueOnAnyThread() && FVisualLogger::IsRecording())
+	if (FVisualLogger::IsRecording())
 	{
 		static const TCHAR* LogName = TEXT("OffsetRootBone");
 		const float InnerCircleRadius = 40.0f;
