UAF strafe warping: Fix crash due to incorrect bone index type

#Synced-CL 50066165

diff --git a/Engine/Plugins/Experimental/UAF/UAFWarping/Source/UAFWarping/Private/StrafeWarpingTrait.cpp b/Engine/Plugins/Experimental/UAF/UAFWarping/Source/UAFWarping/Private/StrafeWarpingTrait.cpp
--- a/Engine/Plugins/Experimental/UAF/UAFWarping/Source/UAFWarping/Private/StrafeWarpingTrait.cpp
+++ b/Engine/Plugins/Experimental/UAF/UAFWarping/Source/UAFWarping/Private/StrafeWarpingTrait.cpp
@@ -564,7 +564,7 @@ void FAnimNextStrafeWarpingTask::DoWarpPose(TUniquePtr<UE::UAF::FKeyframeState>&
 					if (IKTargetIndex != INDEX_NONE)
 					{
 						// Set IK target to equal leg tip in component space
-						int IKTargetParent = Pose.GetLODBoneParentIndex(IKTargetIndex);		// Usually the root bone but we can make it work for any parent
+						FBoneIndexType IKTargetParent = Pose.GetLODBoneParentIndex(IKTargetIndex);		// Usually the root bone but we can make it work for any parent
 						if (IKTargetParent != INDEX_NONE)
 						{
 							FTransform IKTargetCS = FootTargetTransform;
