Fix direct reads of latent SharedData properties in UAF traits

Latent properties on trait SharedData must be accessed via getter methods
(e.g. SharedData->GetAlpha(Binding)) rather than direct field access
(e.g. SharedData->Alpha). Fixed violations in:
- OffsetRootBoneTrait.cpp: TranslationSmoothingTime, RotationSmoothingTime, MaxRotationErrorDegrees
- MotionMatchingTrait.cpp: BlendArguments (x2), BlendProfile
- AnimStateTreeTrait.cpp: StateTreeReference
- MirroringTrait.cpp: Input

#Synced-CL 51590215

diff --git a/Engine/Plugins/Experimental/UAF/UAFMirroring/Source/UAFMirroring/Private/MirroringTrait.cpp b/Engine/Plugins/Experimental/UAF/UAFMirroring/Source/UAFMirroring/Private/MirroringTrait.cpp
--- a/Engine/Plugins/Experimental/UAF/UAFMirroring/Source/UAFMirroring/Private/MirroringTrait.cpp
+++ b/Engine/Plugins/Experimental/UAF/UAFMirroring/Source/UAFMirroring/Private/MirroringTrait.cpp
@@ -118,7 +118,7 @@ namespace UE::UAF
 		
 		if (!InstanceData->Input.IsValid())
 		{
-			InstanceData->Input = Context.AllocateNodeInstance(Binding, SharedData->Input);
+			InstanceData->Input = Context.AllocateNodeInstance(Binding, SharedData->GetInput(Binding));
 		}
 
 		if (!InstanceData->Task.IsValid())
diff --git a/Engine/Plugins/Experimental/UAF/UAFPoseSearch/Source/UAFPoseSearch/Private/MotionMatchingTrait.cpp b/Engine/Plugins/Experimental/UAF/UAFPoseSearch/Source/UAFPoseSearch/Private/MotionMatchingTrait.cpp
--- a/Engine/Plugins/Experimental/UAF/UAFPoseSearch/Source/UAFPoseSearch/Private/MotionMatchingTrait.cpp
+++ b/Engine/Plugins/Experimental/UAF/UAFPoseSearch/Source/UAFPoseSearch/Private/MotionMatchingTrait.cpp
@@ -150,7 +150,7 @@ PRAGMA_ENABLE_DEPRECATION_WARNINGS
 		const TArrayView<const FPoseSearchInteractionAvailability> Availabilities = SharedData->GetAvailabilities(Binding);
 		// MotionMatchInteraction synchronizes and integrates MotionMatchingState.SearchResult internally, so there's no need to update it with the blendstack information
 		UPoseSearchInteractionLibrary::MotionMatchInteraction(MotionMatchingState.SearchResult, Availabilities, AnimContext, FName(), PoseHistory,
-			SharedData->GetbValidateResultAgainstAvailabilities(Binding), SharedData->GetbKeepInteractionAlive(Binding), SharedData->BlendArguments.BlendTime, DeltaTime);
+			SharedData->GetbValidateResultAgainstAvailabilities(Binding), SharedData->GetbKeepInteractionAlive(Binding), SharedData->GetBlendArguments(Binding).BlendTime, DeltaTime);
 
 		// performing the regular single character motion matching search in case there's no MM interaction
 		if (!MotionMatchingState.SearchResult.bIsInteraction)
@@ -385,10 +385,10 @@ PRAGMA_ENABLE_DEPRECATION_WARNINGS
 					{
 						IBlendStack::FGraphRequest NewGraphRequest;
 						NewGraphRequest.FactoryParams = MoveTemp(FactoryParams);
-						NewGraphRequest.BlendArgs = SharedData->BlendArguments;
+						NewGraphRequest.BlendArgs = SharedData->GetBlendArguments(Binding);
 						NewGraphRequest.AnimationGraph = AnimationGraph;
 						NewGraphRequest.FactoryObject = AnimationAsset;
-						NewGraphRequest.BlendProfile = SharedData->BlendProfile;
+						NewGraphRequest.BlendProfile = SharedData->GetBlendProfile(Binding);
 
 						BlendStackBinding.PushGraph(Context, MoveTemp(NewGraphRequest));
 					}
diff --git a/Engine/Plugins/Experimental/UAF/UAFStateTree/Source/UAFStateTree/Private/AnimStateTreeTrait.cpp b/Engine/Plugins/Experimental/UAF/UAFStateTree/Source/UAFStateTree/Private/AnimStateTreeTrait.cpp
--- a/Engine/Plugins/Experimental/UAF/UAFStateTree/Source/UAFStateTree/Private/AnimStateTreeTrait.cpp
+++ b/Engine/Plugins/Experimental/UAF/UAFStateTree/Source/UAFStateTree/Private/AnimStateTreeTrait.cpp
@@ -75,7 +75,7 @@ namespace UE::UAF
 	{
 		const FSharedData* SharedData = Binding.GetSharedData<FSharedData>();
 		FInstanceData* InstanceData = Binding.GetInstanceData<FInstanceData>();
-		InstanceData->StateTree = SharedData->StateTreeReference.GetStateTree();
+		InstanceData->StateTree = SharedData->GetStateTreeReference(Binding).GetStateTree();
 
 		if (InstanceData->StateTree)
 		{
diff --git a/Engine/Plugins/Experimental/UAF/UAFWarping/Source/UAFWarping/Private/OffsetRootBoneTrait.cpp b/Engine/Plugins/Experimental/UAF/UAFWarping/Source/UAFWarping/Private/OffsetRootBoneTrait.cpp
--- a/Engine/Plugins/Experimental/UAF/UAFWarping/Source/UAFWarping/Private/OffsetRootBoneTrait.cpp
+++ b/Engine/Plugins/Experimental/UAF/UAFWarping/Source/UAFWarping/Private/OffsetRootBoneTrait.cpp
@@ -270,7 +270,7 @@ void FAnimNextOffsetRootBoneTask::Execute(UE::UAF::FEvaluationVM& VM) const
 
 		// Blend out translation offset
 		FVector TranslationOffsetDelta = FVector::ZeroVector;
-		FMath::ExponentialSmoothingApprox(TranslationOffsetDelta,TranslationOffset, InstanceData->DeltaTime, SharedData->TranslationSmoothingTime);
+		FMath::ExponentialSmoothingApprox(TranslationOffsetDelta,TranslationOffset, InstanceData->DeltaTime, InstanceData->TranslationSmoothingTime);
 
 		if (SharedData->bClampToTranslationVelocity)
 		{
@@ -293,7 +293,7 @@ void FAnimNextOffsetRootBoneTask::Execute(UE::UAF::FEvaluationVM& VM) const
 		FQuat RotationOffset = ComponentTransform.GetRotation() * InstanceData->SimulatedRotation.Inverse();
 		RotationOffset.EnforceShortestArcWith(FQuat::Identity);
 		FQuat RotationOffsetDelta = FQuat::Identity;
-		SpringMath::ExponentialSmoothingApproxQuat( RotationOffsetDelta, RotationOffset, InstanceData->DeltaTime, SharedData->RotationSmoothingTime);
+		SpringMath::ExponentialSmoothingApproxQuat( RotationOffsetDelta, RotationOffset, InstanceData->DeltaTime, InstanceData->RotationSmoothingTime);
 
 		if (SharedData->bClampToRotationVelocity)
 		{
@@ -316,18 +316,18 @@ void FAnimNextOffsetRootBoneTask::Execute(UE::UAF::FEvaluationVM& VM) const
 		InstanceData->SimulatedRotation = RotationOffsetDelta * InstanceData->SimulatedRotation;
 	}
 
-	if (SharedData->MaxTranslationError >= 0.0f)
+	if (InstanceData->MaxTranslationError >= 0.0f)
 	{
 		FVector TranslationOffset = ComponentTransform.GetLocation() - InstanceData->SimulatedTranslation;
 		const float TranslationOffsetSizeSquared = TranslationOffset.SizeSquared();
-		if (TranslationOffsetSizeSquared > (SharedData->MaxTranslationError * SharedData->MaxTranslationError))
+		if (TranslationOffsetSizeSquared > (InstanceData->MaxTranslationError * InstanceData->MaxTranslationError))
 		{
-			TranslationOffset = TranslationOffset.GetClampedToMaxSize(SharedData->MaxTranslationError);
+			TranslationOffset = TranslationOffset.GetClampedToMaxSize(InstanceData->MaxTranslationError);
 			InstanceData->SimulatedTranslation = ComponentTransform.GetLocation() - TranslationOffset;
 		}
 	}
 
-	const float MaxAngleRadians = FMath::DegreesToRadians(SharedData->MaxRotationErrorDegrees);
+	const float MaxAngleRadians = FMath::DegreesToRadians(InstanceData->MaxRotationErrorDegrees);
 	if (MaxAngleRadians >= 0.0f)
 	{
 		FQuat RotationOffset = ComponentTransform.GetRotation().Inverse() * InstanceData->SimulatedRotation;
@@ -375,9 +375,9 @@ void FAnimNextOffsetRootBoneTask::Execute(UE::UAF::FEvaluationVM& VM) const
 		
 		const UObject* LogOwner = InstanceData->HostObject.Get();
 
-		if (SharedData->MaxTranslationError >= 0.0f)
+		if (InstanceData->MaxTranslationError >= 0.0f)
 		{
-			const float OuterCircleRadius = SharedData->MaxTranslationError + InnerCircleRadius;
+			const float OuterCircleRadius = InstanceData->MaxTranslationError + InnerCircleRadius;
 			UE_VLOG_CIRCLE_THICK(LogOwner, TEXT("OffsetRootBone"), Display, ComponentTransform.GetLocation() + CircleOffset, FVector::UpVector, OuterCircleRadius, FColor::Red, CircleThickness, TEXT(""));
 		}
 		
