UAFWarping - small code cleanup

#Synced-CL 52360741

diff --git a/Engine/Plugins/Experimental/UAF/UAFWarping/Source/UAFWarping/Private/AnimNode/UAFStrafeWarpingNode.cpp b/Engine/Plugins/Experimental/UAF/UAFWarping/Source/UAFWarping/Private/AnimNode/UAFStrafeWarpingNode.cpp
--- a/Engine/Plugins/Experimental/UAF/UAFWarping/Source/UAFWarping/Private/AnimNode/UAFStrafeWarpingNode.cpp
+++ b/Engine/Plugins/Experimental/UAF/UAFWarping/Source/UAFWarping/Private/AnimNode/UAFStrafeWarpingNode.cpp
@@ -339,8 +339,6 @@ void FUAFStrafeWarpingPostAnimOp::DoWarpPose(FUAFAnimOpValueEvaluator& Evaluator
 
 	const FAttributeMappingKey MappingKey = FAttributeMappingKey::MakeFromTo<FBoneTransformAnimationAttribute>();
 
-	const FValueBundlePtr DefaultValues = ValueBundle.GetNamedSet()->GetDefaultAttributeValues();
-
 	TBoundValueMap<FBoneTransformAnimationAttribute>* BoneTransforms = ValueBundle.GetBoundValueMaps().Find<FBoneTransformAnimationAttribute>(MappingKey);
 	if (!BoneTransforms)
 	{
