UAFWarping - WarpToTargetNode

#Synced-CL 52486799

diff --git a/Engine/Plugins/Experimental/UAF/UAFWarping/Source/UAFWarping/Private/AnimNode/UAFWarpToTargetNode.cpp b/Engine/Plugins/Experimental/UAF/UAFWarping/Source/UAFWarping/Private/AnimNode/UAFWarpToTargetNode.cpp
new file mode 100644
--- /dev/null
+++ b/Engine/Plugins/Experimental/UAF/UAFWarping/Source/UAFWarping/Private/AnimNode/UAFWarpToTargetNode.cpp
@@ -0,0 +1,119 @@
+// Copyright Epic Games, Inc. All Rights Reserved.
+
+#include "UAFWarpToTargetNode.h"
+#include "Animation/AnimRootMotionProvider.h"
+#include "AnimNextWarpingLog.h"
+#include "TwoBoneIK.h"
+#include "UAF/AnimNodeCore/UAFAnimNodeUpdate.h"
+#include "UAF/AnimOpCore/UAFAnimOpValueEvaluator.h"
+#include "UAF/ValueRuntime/Transformers/BoneSpace.h"
+#include "VisualLogger/VisualLogger.h"
+
+namespace UE::UAF
+{
+
+FUAFWarpToTargetPostAnimOp::FUAFWarpToTargetPostAnimOp()
+: FUAFAnimOp(1)
+{
+	InitializeAs<FUAFWarpToTargetPostAnimOp>();
+}
+
+void FUAFWarpToTargetPostAnimOp::EvaluateValues(FUAFAnimOpValueEvaluator& Evaluator)
+{
+	const FUAFAnimOpValueEvaluationContext& AnimOpEvaluationContext = Evaluator.GetActiveEvaluationContext();
+
+	const USkeleton* Skeleton = AnimOpEvaluationContext.GetSkeleton();
+	FValueBundleCoWRef* InputRef = Evaluator.GetEvaluationStack().PeekMutable();
+	FValueBundle& ValueBundle = InputRef->GetMutable();
+
+	const FAttributeNamedSetPtr NamedSet = ValueBundle.GetNamedSet();
+	if (!NamedSet)
+	{
+		UE_LOG(LogAnimNextWarping, Error, TEXT("FUAFWarpToTargetPostAnimOp::EvaluateValues, NamedSet not found"));
+		return;
+	}
+
+	const FAttributeTypedSetPtr AttributeTypedSet = NamedSet->FindTypedSet<FTransformAnimationAttribute>();
+	if (!AttributeTypedSet)
+	{
+		UE_LOG(LogAnimNextWarping, Error, TEXT("FUAFWarpToTargetPostAnimOp::EvaluateValues, AttributeTypedSet not found"));
+		return;
+	}
+
+	const FAttributeSetIndex SetIndex = AttributeTypedSet->FindIndex(UE::Anim::IAnimRootMotionProvider::AttributeName);
+	if (!SetIndex)
+	{
+		UE_LOG(LogAnimNextWarping, Error, TEXT("FUAFWarpToTargetPostAnimOp::EvaluateValues, SetIndex for IAnimRootMotionProvider not found"));
+		return;
+	}
+
+	// This attribute is present within the current named set
+	const FAttributeMappingKey MappingKey = FAttributeMappingKey::MakeFromTo<FTransformAnimationAttribute>();
+	TBoundValueMap<FTransformAnimationAttribute>* Map = ValueBundle.GetBoundValueMaps().Find<FTransformAnimationAttribute>(MappingKey);
+
+	if (!Map)
+	{
+		UE_LOG(LogAnimNextWarping, Error, TEXT("FUAFWarpToTargetPostAnimOp::EvaluateValues, This attribute exists within our named set but isn't present in the output value bundle"));
+		return;
+	}
+
+	FTransform& RootMotionTransform = (*Map)[SetIndex].Value;
+	RootMotionTransform = TargetRootBoneTransform.GetRelativeTransform(RootBoneTransform);
+	
+#if ENABLE_VISUAL_LOG && ENABLE_ANIM_DEBUG
+	if (FVisualLogger::IsRecording())
+	{
+		if (const UObject* VLogObject = HostObject.Pin().Get())
+		{
+			UE_VLOG_COORDINATESYSTEM(VLogObject, "WarpToTarget", Display, RootBoneTransform.GetLocation(), RootBoneTransform.GetRotation().Rotator(), 10.f, FColor::Red, 1, TEXT(""));
+			UE_VLOG_COORDINATESYSTEM(VLogObject, "WarpToTarget", Display, TargetRootBoneTransform.GetLocation(), TargetRootBoneTransform.GetRotation().Rotator(), 50.f, FColor::Green, 1, TEXT(""));
+		}
+	}
+#endif // ENABLE_VISUAL_LOG && ENABLE_ANIM_DEBUG
+}
+
+///////////////////////////////////////////////////////////
+// FUAFWarpToTargetNode
+FUAFWarpToTargetNode::FUAFWarpToTargetNode(FUAFAnimGraphUpdateContext& Context, const FUAFWarpToTargetData& InData)
+	: FUAFModifierAnimNode(Context)
+{
+	InitializeAs<FUAFWarpToTargetNode>(Context);
+	InitializeModifier(Context, InData);
+
+	Data = &InData;
+
+	if (HasChild())
+	{
+		PostAnimOp.HostObject = Context.GetHostObject();
+		PostAnimOp.SetDebugOwner(this);
+
+		SetPostAnimOp(&PostAnimOp);
+	}
+}
+
+void FUAFWarpToTargetNode::PreUpdate(FUAFAnimGraphUpdateContext& Context)
+{
+	check(Data);
+	PostAnimOp.RootBoneTransform = Data->RootBoneTransform.GetValue(Context.GetVariablesOwner());
+	PostAnimOp.TargetRootBoneTransform = Data->TargetRootBoneTransform.GetValue(Context.GetVariablesOwner());
+}
+
+#if UAF_TRACE_ENABLED
+FString FUAFWarpToTargetNode::GetDebugName() const
+{
+	static FString Name("Warp To Target");
+	return Name;
+}
+
+UStruct* FUAFWarpToTargetNode::GetDebugStruct() const
+{
+	return FUAFWarpToTargetData::StaticStruct();
+}
+#endif
+
+FUAFAnimNodePtr FUAFWarpToTargetData::CreateInstance(FUAFAnimGraphUpdateContext& Context) const
+{
+	return MakeAnimNode<FUAFWarpToTargetNode>(Context, *this);
+}
+
+} // namespace UE::UAF
diff --git a/Engine/Plugins/Experimental/UAF/UAFWarping/Source/UAFWarping/Private/AnimNode/UAFWarpToTargetNode.h b/Engine/Plugins/Experimental/UAF/UAFWarping/Source/UAFWarping/Private/AnimNode/UAFWarpToTargetNode.h
new file mode 100644
--- /dev/null
+++ b/Engine/Plugins/Experimental/UAF/UAFWarping/Source/UAFWarping/Private/AnimNode/UAFWarpToTargetNode.h
@@ -0,0 +1,64 @@
+// Copyright Epic Games, Inc. All Rights Reserved.
+
+#pragma once
+
+#include "BindableValue/UAFBindableTypes.h"
+#include "CoreMinimal.h"
+#include "UAF/AnimNodeCore/UAFModifierAnimNode.h"
+#include "UAF/AnimNodeCore/UAFModifierAnimNodeData.h"
+#include "UAF/Attributes/AttributeTypedSet.h"
+
+#include "UAFWarpToTargetNode.generated.h"
+
+namespace UE::UAF
+{
+
+USTRUCT()
+struct FUAFWarpToTargetPostAnimOp : public FUAFAnimOp
+{
+	GENERATED_BODY()
+	UAF_DECLARE_ANIMOP(FUAFWarpToTargetPostAnimOp)
+
+	FUAFWarpToTargetPostAnimOp();
+
+	virtual void EvaluateValues(FUAFAnimOpValueEvaluator& Evaluator) override;
+
+	TWeakObjectPtr<const UObject> HostObject;
+
+	FTransform RootBoneTransform = FTransform::Identity;
+	FTransform TargetRootBoneTransform = FTransform::Identity;
+};
+
+USTRUCT(DisplayName = "Warp To Target")
+struct FUAFWarpToTargetData : public FUAFModifierAnimNodeData
+{
+	GENERATED_BODY()
+
+	// @TODO Temp / try to remove this. Shouldn't have to feed as argument
+	// Last root bone transform sampled
+	UPROPERTY(EditAnywhere, Category = Evaluation)
+	FBindableTransform RootBoneTransform;
+
+	UPROPERTY(EditAnywhere, Category = Evaluation)
+	FBindableTransform TargetRootBoneTransform;
+
+	// FUAFAnimNodeData impl
+	virtual FUAFAnimNodePtr CreateInstance(FUAFAnimGraphUpdateContext& Context) const override;
+};
+
+struct FUAFWarpToTargetNode : FUAFModifierAnimNode
+{
+	FUAFWarpToTargetNode(FUAFAnimGraphUpdateContext& Context, const FUAFWarpToTargetData& InData);
+
+	virtual void PreUpdate(FUAFAnimGraphUpdateContext& Context) override;
+
+#if UAF_TRACE_ENABLED
+	virtual FString GetDebugName() const override;
+	virtual UStruct* GetDebugStruct() const override;
+#endif
+private:
+	FUAFWarpToTargetPostAnimOp PostAnimOp;
+	const FUAFWarpToTargetData* Data = nullptr;
+};
+
+} // namespace UE::UAF
\ No newline at end of file
